Virtual Production Blank Project Config files update.  #jira UEENT-2869 #homam.bahnassi

diff --git a/Templates/TP_ME_BlankBP/Config/DefaultEngine.ini b/Templates/TP_ME_BlankBP/Config/DefaultEngine.ini
--- a/Templates/TP_ME_BlankBP/Config/DefaultEngine.ini
+++ b/Templates/TP_ME_BlankBP/Config/DefaultEngine.ini
@@ -32,22 +32,21 @@ r.DefaultFeature.LensFlare=True
 r.ClearCoatNormal=False
 r.NormalMapsForStaticLighting=False
 
+;Sequencer.AlwayRecordLiveLinkTimecode 1
 
 
 [/Script/EngineSettings.GameMapsSettings]
 ;defines what level to be loaded by default in the editor and at playback time
-GameDefaultMap=/Game/VirtualProduction/Main.Main
-EditorStartupMap=/Game/VirtualProduction/Main.Main
+GameDefaultMap=/Game/VProdProject/Main.Main
+EditorStartupMap=/Game/VProdProject/Main.Main
 
 [/Script/UnrealEd.UnrealEdEngine]
 ; Remove Engine Template maps
--TemplateMapInfos=(ThumbnailTexture=Texture2D'/Engine/Maps/Templates/Thumbnails/Default.Default',Map="/Engine/Maps/Templates/Template_Default")
--TemplateMapInfos=(ThumbnailTexture=Texture2D'/Engine/Maps/Templates/Thumbnails/TimeOfDay.TimeOfDay',Map="/Engine/Maps/Templates/TimeOfDay_Default")
+;TemplateMapInfos=(ThumbnailTexture=Texture2D'/Engine/Maps/Templates/Thumbnails/Default.Default',Map="/Engine/Maps/Templates/Template_Default")
+;TemplateMapInfos=(ThumbnailTexture=Texture2D'/Engine/Maps/Templates/Thumbnails/TimeOfDay.TimeOfDay',Map="/Engine/Maps/Templates/TimeOfDay_Default")
 -TemplateMapInfos=(ThumbnailTexture=Texture2D'/Engine/Maps/Templates/Thumbnails/VR-Basic.VR-Basic',Map="/Engine/Maps/Templates/VR-Basic")
 
-; Add Predefined Template Maps that are suitable for this template
-;+TemplateMapInfos=(ThumbnailTexture=Texture2D'/Game/Developers/Maps/Thumbnails/AssetName.AssetName',Map="/Game/Maps/Templates/MapName")
-+TemplateMapInfos=(ThumbnailTexture=Texture2D'/Game/LevelTemplates/Default/Thumbnails/Default.Default',Map="/Game/LevelTemplates/Default/Default")
+
 
 
 
@@ -76,10 +75,11 @@ VirtualScoutingUI=/VirtualProductionUtilities/Editor/VirtualScoutingWidget.Virtu
 ;Setup for Media Profiles
 [/Script/MediaFrameworkUtilities.MediaProfileSettings]
 bApplyInCommandlet=False
-MediaSourceProxy=/Game/VProd/MediaIO/VideoProxyInA.VideoProxyInA
-MediaSourceProxy=/Game/VProd/MediaIO/VideoProxyInB.VideoProxyInB
-MediaOutputProxy=/Game/VProd/MediaIO/ProxyMediaOutput.ProxyMediaOutput
-StartupMediaProfile=/Game/VProd/MediaIO/AJA_MediaProfile.AJA_MediaProfile
+MediaSourceProxy=/Game/VprodProject/MediaIO/MediaSource-01.MediaSource-01
+MediaSourceProxy=/Game/VprodProject/MediaIO/MediaSource-02.MediaSource-02
+MediaOutputProxy=/Game/VprodProject/MediaIO/MediaOutput-01.MediaOutput-01
+StartupMediaProfile=None
+
 
 ;Setup for Multiuser
 [/Script/Concert.ConcertClientConfig]
@@ -93,6 +93,16 @@ DefaultSaveSessionAs=
 ClientSettings=(DisplayName="",AvatarColor=(R=1.000000,G=1.000000,B=1.000000,A=1.000000),DesktopAvatarActorClass=/ConcertSyncClient/DesktopPresence.DesktopPresence_C,VRAvatarActorClass=/ConcertSyncClient/VRPresence.VRPresence_C,DiscoveryTimeoutSeconds=5,SessionTickFrequencySeconds=1,LatencyCompensationMs=0.000000,Tags=)
 EndpointSettings=(bEnableLogging=False,PurgeProcessedMessageDelaySeconds=30,RemoteEndpointTimeoutSeconds=60)
 
+;Setup for Ndisplay
+[/Script/Engine.Engine]
+GameEngine=/Script/DisplayCluster.DisplayClusterGameEngine
+UnrealEdEngine=/Script/DisplayClusterEditor.DisplayClusterEditorEngine
+GameViewportClientClassName=/Script/DisplayCluster.DisplayClusterViewportClient
+
+;Setup for Multiuser
+[/Script/DisplayClusterEditor.DisplayClusterEditorSettings]
+bEnabled=True
+
 [/Script/Engine.PhysicsSettings]
 DefaultGravityZ=-980.000000
 DefaultTerminalVelocity=4000.000000
diff --git a/Templates/TP_ME_BlankBP/Config/DefaultGame.ini b/Templates/TP_ME_BlankBP/Config/DefaultGame.ini
--- a/Templates/TP_ME_BlankBP/Config/DefaultGame.ini
+++ b/Templates/TP_ME_BlankBP/Config/DefaultGame.ini
@@ -1 +1,11 @@
+[/Script/EngineSettings.GeneralProjectSettings]
+bUseBorderlessWindow=True
+
+[/Script/MediaFrameworkUtilities.MediaProfileSettings]
+bApplyInCommandlet=False
++MediaSourceProxy=/Game/VprodProject/MediaIO/MediaSource-01.MediaSource-01
++MediaSourceProxy=/Game/VprodProject/MediaIO/MediaSource-02.MediaSource-02
++MediaOutputProxy=/Game/VprodProject/MediaIO/MediaOutput-01.MediaOutput-01
+StartupMediaProfile=None
+
 
