Update to Virtual Production Project Templates.  Saved Maps in proper state.  Thumbnail captures.  Save VirtualProduction.ini

diff --git a/Templates/TP_ME_BlankBP/Config/DefaultEngine.ini b/Templates/TP_ME_BlankBP/Config/DefaultEngine.ini
--- a/Templates/TP_ME_BlankBP/Config/DefaultEngine.ini
+++ b/Templates/TP_ME_BlankBP/Config/DefaultEngine.ini
@@ -13,7 +13,7 @@ r.DefaultFeature.AutoExposure.ExtendDefaultLuminanceRange=True
 ;r.DefaultFeature.AutoExposure=False
 ;r.DefaultFeature.AutoExposure.Method=0
 
-;dynamic gi settings
+; dynamic gi settings
 r.LightPropagationVolume=0
 
 ; quality settings
@@ -26,42 +26,48 @@ r.SSR.Quality=4
 r.GBufferFormat=3
 
 ; shader settings
-;r.PostProcessing.PropagateAlpha=1
+r.PostProcessing.PropagateAlpha=1
 r.SupportSkyAtmosphereAffectsHeightFog=True
 r.DefaultFeature.LensFlare=True
 r.ClearCoatNormal=False
 r.NormalMapsForStaticLighting=False
 
-;Sequencer.AlwayRecordLiveLinkTimecode 1
+; Sequencer Take Recorder
+Sequencer.AlwayRecordLiveLinkTimecode 1
 
 
 [/Script/EngineSettings.GameMapsSettings]
 ;defines what level to be loaded by default in the editor and at playback time
-GameDefaultMap=/Game/VProdProject/Main.Main
-EditorStartupMap=/Game/VProdProject/Main.Main
+GameDefaultMap=/Game/Main.Main
+EditorStartupMap=/Game/Main.Main
 
 [/Script/UnrealEd.UnrealEdEngine]
 ; Remove Engine Template maps
-;TemplateMapInfos=(ThumbnailTexture=Texture2D'/Engine/Maps/Templates/Thumbnails/Default.Default',Map="/Engine/Maps/Templates/Template_Default")
-;TemplateMapInfos=(ThumbnailTexture=Texture2D'/Engine/Maps/Templates/Thumbnails/TimeOfDay.TimeOfDay',Map="/Engine/Maps/Templates/TimeOfDay_Default")
 -TemplateMapInfos=(ThumbnailTexture=Texture2D'/Engine/Maps/Templates/Thumbnails/VR-Basic.VR-Basic',Map="/Engine/Maps/Templates/VR-Basic")
 
 
-
-
-
-;Allows for Hardware Accelerated Video Decoding
+; Allows for Hardware Accelerated Video Decoding
 [/Script/WmfMediaFactory.WmfMediaSettings]
 AllowNonStandardCodecs=True
 HardwareAcceleratedVideoDecoding=True
 LowLatency=False
 NativeAudioOut=False
 
-;Adds Virtual Scouting Widget for VR Scouting
+; Adds Virtual Scouting Widget for VR Scouting
+
 [/Script/VPUtilitiesEditor.VPUtilitiesEditorSettings]
 VirtualScoutingUI=/VirtualProductionUtilities/Editor/VirtualScoutingWidget.VirtualScoutingWidget_C
+FlightSpeed=0.500000
+GripNavSpeed=0.250000
+bUseMetric=False
+bUseTransformGizmo=False
+bUseGripInertiaDamping=True
+InertiaDamping=0.950000
+bIsHelperSystemEnabled=True
+ScoutingSubsystemEdititorUtilityActorClassPath=/VirtualProductionUtilities/VirtualProductionHelpers.VirtualProductionHelpers_C
 
-;Enables WebControl API
+
+; Enables WebControl API
 
 ;WebControl.EnableServerOnStartup=True
 ;WebControl.EnableServerOnStartup=1
@@ -72,16 +78,7 @@ VirtualScoutingUI=/VirtualProductionUtilities/Editor/VirtualScoutingWidget.Virtu
 +Channels=(Name=FRemoteSessionInputChannel,Mode=Read)
 +Channels=(Name=FRemoteSessionXRTrackingChannel,Mode=Read)
 
-;Setup for Media Profiles
-[/Script/MediaFrameworkUtilities.MediaProfileSettings]
-bApplyInCommandlet=False
-MediaSourceProxy=/Game/VprodProject/MediaIO/MediaSource-01.MediaSource-01
-MediaSourceProxy=/Game/VprodProject/MediaIO/MediaSource-02.MediaSource-02
-MediaOutputProxy=/Game/VprodProject/MediaIO/MediaOutput-01.MediaOutput-01
-StartupMediaProfile=None
-
-
-;Setup for Multiuser
+; Setup for Multiuser
 [/Script/Concert.ConcertClientConfig]
 bIsHeadless=False
 bInstallEditorToolbarButton=True
@@ -93,60 +90,15 @@ DefaultSaveSessionAs=
 ClientSettings=(DisplayName="",AvatarColor=(R=1.000000,G=1.000000,B=1.000000,A=1.000000),DesktopAvatarActorClass=/ConcertSyncClient/DesktopPresence.DesktopPresence_C,VRAvatarActorClass=/ConcertSyncClient/VRPresence.VRPresence_C,DiscoveryTimeoutSeconds=5,SessionTickFrequencySeconds=1,LatencyCompensationMs=0.000000,Tags=)
 EndpointSettings=(bEnableLogging=False,PurgeProcessedMessageDelaySeconds=30,RemoteEndpointTimeoutSeconds=60)
 
-;Setup for Ndisplay
+; Setup for Ndisplay
 [/Script/Engine.Engine]
 GameEngine=/Script/DisplayCluster.DisplayClusterGameEngine
 UnrealEdEngine=/Script/DisplayClusterEditor.DisplayClusterEditorEngine
 GameViewportClientClassName=/Script/DisplayCluster.DisplayClusterViewportClient
 
-;Setup for Multiuser
+; Setup for Ndisplay
 [/Script/DisplayClusterEditor.DisplayClusterEditorSettings]
 bEnabled=True
 
-[/Script/Engine.PhysicsSettings]
-DefaultGravityZ=-980.000000
-DefaultTerminalVelocity=4000.000000
-DefaultFluidFriction=0.300000
-SimulateScratchMemorySize=262144
-RagdollAggregateThreshold=4
-TriangleMeshTriangleMinAreaThreshold=5.000000
-bEnableShapeSharing=False
-bEnablePCM=True
-bEnableStabilization=False
-bWarnMissingLocks=True
-bEnable2DPhysics=False
-PhysicErrorCorrection=(PingExtrapolation=0.100000,PingLimit=100.000000,ErrorPerLinearDifference=1.000000,ErrorPerAngularDifference=1.000000,MaxRestoredStateError=1.000000,MaxLinearHardSnapDistance=400.000000,PositionLerp=0.000000,AngleLerp=0.400000,LinearVelocityCoefficient=100.000000,AngularVelocityCoefficient=10.000000,ErrorAccumulationSeconds=0.500000,ErrorAccumulationDistanceSq=15.000000,ErrorAccumulationSimilarity=100.000000)
-LockedAxis=Invalid
-DefaultDegreesOfFreedom=Full3D
-BounceThresholdVelocity=200.000000
-FrictionCombineMode=Average
-RestitutionCombineMode=Average
-MaxAngularVelocity=3600.000000
-MaxDepenetrationVelocity=0.000000
-ContactOffsetMultiplier=0.020000
-MinContactOffset=2.000000
-MaxContactOffset=8.000000
-bSimulateSkeletalMeshOnDedicatedServer=True
-DefaultShapeComplexity=CTF_UseSimpleAndComplex
-bDefaultHasComplexCollision=True
-bSuppressFaceRemapTable=False
-bSupportUVFromHitResults=False
-bDisableActiveActors=False
-bDisableKinematicStaticPairs=False
-bDisableKinematicKinematicPairs=False
-bDisableCCD=False
-bEnableEnhancedDeterminism=False
-AnimPhysicsMinDeltaTime=0.000000
-bSimulateAnimPhysicsAfterReset=False
-MaxPhysicsDeltaTime=0.033333
-bSubstepping=False
-bSubsteppingAsync=False
-MaxSubstepDeltaTime=0.016667
-MaxSubsteps=6
-SyncSceneSmoothingFactor=0.000000
-InitialAverageFrameRate=0.016667
-PhysXTreeRebuildRate=10
-DefaultBroadphaseSettings=(bUseMBPOnClient=False,bUseMBPOnServer=False,bUseMBPOuterBounds=False,MBPBounds=(Min=(X=0.000000,Y=0.000000,Z=0.000000),Max=(X=0.000000,Y=0.000000,Z=0.000000),IsValid=0),MBPOuterBounds=(Min=(X=0.000000,Y=0.000000,Z=0.000000),Max=(X=0.000000,Y=0.000000,Z=0.000000),IsValid=0),MBPNumSubdivs=2)
-ChaosSettings=(DefaultThreadingModel=TaskGraph,DedicatedThreadTickMode=VariableCappedWithTarget,DedicatedThreadBufferMode=Double)
 
 
diff --git a/Templates/TP_ME_BlankBP/Config/DefaultVirtualProductionUtilities.ini b/Templates/TP_ME_BlankBP/Config/DefaultVirtualProductionUtilities.ini
new file mode 100644
--- /dev/null
+++ b/Templates/TP_ME_BlankBP/Config/DefaultVirtualProductionUtilities.ini
@@ -0,0 +1,13 @@
+
+[/Script/VPUtilitiesEditor.VPUtilitiesEditorSettings]
+VirtualScoutingUI=/VirtualProductionUtilities/Editor/VirtualScoutingWidget.VirtualScoutingWidget_C
+FlightSpeed=0.500000
+GripNavSpeed=0.250000
+bUseMetric=False
+bUseTransformGizmo=False
+bUseGripInertiaDamping=True
+InertiaDamping=0.950000
+bIsHelperSystemEnabled=True
+ScoutingSubsystemEdititorUtilityActorClassPath=/VirtualProductionUtilities/VirtualProductionHelpers.VirtualProductionHelpers_C
+
+
diff --git a/Templates/TP_ME_VProdBP/Config/DefaultEngine.ini b/Templates/TP_ME_VProdBP/Config/DefaultEngine.ini
--- a/Templates/TP_ME_VProdBP/Config/DefaultEngine.ini
+++ b/Templates/TP_ME_VProdBP/Config/DefaultEngine.ini
@@ -26,13 +26,14 @@ r.SSR.Quality=4
 r.GBufferFormat=3
 
 ; shader settings
-;r.PostProcessing.PropagateAlpha=1
+r.PostProcessing.PropagateAlpha=1
 r.SupportSkyAtmosphereAffectsHeightFog=True
 r.DefaultFeature.LensFlare=True
 r.ClearCoatNormal=False
 r.NormalMapsForStaticLighting=False
 
-;Sequencer.AlwayRecordLiveLinkTimecode 1
+; Sequencer Take Recorder
+Sequencer.AlwayRecordLiveLinkTimecode 1
 
 
 [/Script/EngineSettings.GameMapsSettings]
@@ -42,14 +43,9 @@ EditorStartupMap=/Game/VProdProject/Main.Main
 
 [/Script/UnrealEd.UnrealEdEngine]
 ; Remove Engine Template maps
-;TemplateMapInfos=(ThumbnailTexture=Texture2D'/Engine/Maps/Templates/Thumbnails/Default.Default',Map="/Engine/Maps/Templates/Template_Default")
-;TemplateMapInfos=(ThumbnailTexture=Texture2D'/Engine/Maps/Templates/Thumbnails/TimeOfDay.TimeOfDay',Map="/Engine/Maps/Templates/TimeOfDay_Default")
 -TemplateMapInfos=(ThumbnailTexture=Texture2D'/Engine/Maps/Templates/Thumbnails/VR-Basic.VR-Basic',Map="/Engine/Maps/Templates/VR-Basic")
 
 
-
-
-
 ;Allows for Hardware Accelerated Video Decoding
 [/Script/WmfMediaFactory.WmfMediaSettings]
 AllowNonStandardCodecs=True
@@ -103,50 +99,5 @@ GameViewportClientClassName=/Script/DisplayCluster.DisplayClusterViewportClient
 [/Script/DisplayClusterEditor.DisplayClusterEditorSettings]
 bEnabled=True
 
-[/Script/Engine.PhysicsSettings]
-DefaultGravityZ=-980.000000
-DefaultTerminalVelocity=4000.000000
-DefaultFluidFriction=0.300000
-SimulateScratchMemorySize=262144
-RagdollAggregateThreshold=4
-TriangleMeshTriangleMinAreaThreshold=5.000000
-bEnableShapeSharing=False
-bEnablePCM=True
-bEnableStabilization=False
-bWarnMissingLocks=True
-bEnable2DPhysics=False
-PhysicErrorCorrection=(PingExtrapolation=0.100000,PingLimit=100.000000,ErrorPerLinearDifference=1.000000,ErrorPerAngularDifference=1.000000,MaxRestoredStateError=1.000000,MaxLinearHardSnapDistance=400.000000,PositionLerp=0.000000,AngleLerp=0.400000,LinearVelocityCoefficient=100.000000,AngularVelocityCoefficient=10.000000,ErrorAccumulationSeconds=0.500000,ErrorAccumulationDistanceSq=15.000000,ErrorAccumulationSimilarity=100.000000)
-LockedAxis=Invalid
-DefaultDegreesOfFreedom=Full3D
-BounceThresholdVelocity=200.000000
-FrictionCombineMode=Average
-RestitutionCombineMode=Average
-MaxAngularVelocity=3600.000000
-MaxDepenetrationVelocity=0.000000
-ContactOffsetMultiplier=0.020000
-MinContactOffset=2.000000
-MaxContactOffset=8.000000
-bSimulateSkeletalMeshOnDedicatedServer=True
-DefaultShapeComplexity=CTF_UseSimpleAndComplex
-bDefaultHasComplexCollision=True
-bSuppressFaceRemapTable=False
-bSupportUVFromHitResults=False
-bDisableActiveActors=False
-bDisableKinematicStaticPairs=False
-bDisableKinematicKinematicPairs=False
-bDisableCCD=False
-bEnableEnhancedDeterminism=False
-AnimPhysicsMinDeltaTime=0.000000
-bSimulateAnimPhysicsAfterReset=False
-MaxPhysicsDeltaTime=0.033333
-bSubstepping=False
-bSubsteppingAsync=False
-MaxSubstepDeltaTime=0.016667
-MaxSubsteps=6
-SyncSceneSmoothingFactor=0.000000
-InitialAverageFrameRate=0.016667
-PhysXTreeRebuildRate=10
-DefaultBroadphaseSettings=(bUseMBPOnClient=False,bUseMBPOnServer=False,bUseMBPOuterBounds=False,MBPBounds=(Min=(X=0.000000,Y=0.000000,Z=0.000000),Max=(X=0.000000,Y=0.000000,Z=0.000000),IsValid=0),MBPOuterBounds=(Min=(X=0.000000,Y=0.000000,Z=0.000000),Max=(X=0.000000,Y=0.000000,Z=0.000000),IsValid=0),MBPNumSubdivs=2)
-ChaosSettings=(DefaultThreadingModel=TaskGraph,DedicatedThreadTickMode=VariableCappedWithTarget,DedicatedThreadBufferMode=Double)
 
 
diff --git a/Templates/TP_ME_VProdBP/Config/DefaultVirtualProductionUtilities.ini b/Templates/TP_ME_VProdBP/Config/DefaultVirtualProductionUtilities.ini
new file mode 100644
--- /dev/null
+++ b/Templates/TP_ME_VProdBP/Config/DefaultVirtualProductionUtilities.ini
@@ -0,0 +1,13 @@
+
+[/Script/VPUtilitiesEditor.VPUtilitiesEditorSettings]
+VirtualScoutingUI=/VirtualProductionUtilities/Editor/VirtualScoutingWidget.VirtualScoutingWidget_C
+FlightSpeed=0.500000
+GripNavSpeed=0.250000
+bUseMetric=False
+bUseTransformGizmo=False
+bUseGripInertiaDamping=True
+InertiaDamping=0.950000
+bIsHelperSystemEnabled=True
+ScoutingSubsystemEdititorUtilityActorClassPath=/VirtualProductionUtilities/VirtualProductionHelpers.VirtualProductionHelpers_C
+
+
