Updating Sequencer.AlwayRecordLiveLinkTimecode 1 to Sequencer.AlwayRecordLiveLinkTimecode=1 to resolve warnings in CIS.

diff --git a/Templates/TP_ME_BlankBP/Config/DefaultEngine.ini b/Templates/TP_ME_BlankBP/Config/DefaultEngine.ini
--- a/Templates/TP_ME_BlankBP/Config/DefaultEngine.ini
+++ b/Templates/TP_ME_BlankBP/Config/DefaultEngine.ini
@@ -36,7 +36,7 @@ r.ClearCoatNormal=False
 r.NormalMapsForStaticLighting=False
 
 ; Sequencer Take Recorder
-Sequencer.AlwayRecordLiveLinkTimecode 1
+Sequencer.AlwayRecordLiveLinkTimecode=1
 
 
 [/Script/EngineSettings.GameMapsSettings]
diff --git a/Templates/TP_ME_VProdBP/Config/DefaultEngine.ini b/Templates/TP_ME_VProdBP/Config/DefaultEngine.ini
--- a/Templates/TP_ME_VProdBP/Config/DefaultEngine.ini
+++ b/Templates/TP_ME_VProdBP/Config/DefaultEngine.ini
@@ -44,7 +44,7 @@ r.ClearCoatNormal=False
 r.NormalMapsForStaticLighting=False
 
 ; Sequencer Take Recorder
-Sequencer.AlwayRecordLiveLinkTimecode 1
+Sequencer.AlwayRecordLiveLinkTimecode=1
 
 
 [/Script/EngineSettings.GameMapsSettings]
