Virutal Production Blank Project Template.  Removed Developer folders and LevelTemplate folders.  Moved Main.Map to Map folder.  Changed DefaultEngine.ini to new Main map location.

diff --git a/Templates/TP_ME_BlankBP/Config/DefaultEngine.ini b/Templates/TP_ME_BlankBP/Config/DefaultEngine.ini
--- a/Templates/TP_ME_BlankBP/Config/DefaultEngine.ini
+++ b/Templates/TP_ME_BlankBP/Config/DefaultEngine.ini
@@ -38,8 +38,8 @@ r.NormalMapsForStaticLighting=False
 
 [/Script/EngineSettings.GameMapsSettings]
 ;defines what level to be loaded by default in the editor and at playback time
-GameDefaultMap=/Game/VProdProject/Maps/Ndisplay.Ndisplay
-EditorStartupMap=/Game/VProdProject/Maps/Main.Main
+EditorStartupMap=/Game/Maps/Main.Main
+EditorStartupMap=/Game/Maps/Main.Main
 
 [/Script/UnrealEd.UnrealEdEngine]
 ; Remove Engine Template maps
