Virutal Production Project Template.  Update VP content.  2 sided materials.  Removing assets from Blank Template.

diff --git a/Templates/TP_ME_BlankBP/Config/DefaultEngine.ini b/Templates/TP_ME_BlankBP/Config/DefaultEngine.ini
--- a/Templates/TP_ME_BlankBP/Config/DefaultEngine.ini
+++ b/Templates/TP_ME_BlankBP/Config/DefaultEngine.ini
@@ -38,8 +38,8 @@ r.NormalMapsForStaticLighting=False
 
 [/Script/EngineSettings.GameMapsSettings]
 ;defines what level to be loaded by default in the editor and at playback time
-EditorStartupMap=/Game/Maps/Main.Main
-EditorStartupMap=/Game/Maps/Main.Main
+EditorStartupMap=/Game/Main.Main
+EditorStartupMap=/Game/Main.Main
 
 [/Script/UnrealEd.UnrealEdEngine]
 ; Remove Engine Template maps
