- Fixed Blank Virtual Production Template referencing deleted assets

#rnx

diff --git a/Templates/TP_ME_BlankBP/Config/DefaultGame.ini b/Templates/TP_ME_BlankBP/Config/DefaultGame.ini
--- a/Templates/TP_ME_BlankBP/Config/DefaultGame.ini
+++ b/Templates/TP_ME_BlankBP/Config/DefaultGame.ini
@@ -1,11 +1,4 @@
 [/Script/EngineSettings.GeneralProjectSettings]
 bUseBorderlessWindow=True
 
-[/Script/MediaFrameworkUtilities.MediaProfileSettings]
-bApplyInCommandlet=False
-+MediaSourceProxy=/Game/VprodProject/MediaIO/MediaSource-01.MediaSource-01
-+MediaSourceProxy=/Game/VprodProject/MediaIO/MediaSource-02.MediaSource-02
-+MediaOutputProxy=/Game/VprodProject/MediaIO/MediaOutput-01.MediaOutput-01
-StartupMediaProfile=None
-
 
