Adding Chinese translation to this file. #jira ueloc-4711 #fyi jamie.dale, michelle.zhao, alexis.zhang, yali.he

diff --git a/Templates/TP_ME_BlankBP/Config/TemplateDefs.ini b/Templates/TP_ME_BlankBP/Config/TemplateDefs.ini
--- a/Templates/TP_ME_BlankBP/Config/TemplateDefs.ini
+++ b/Templates/TP_ME_BlankBP/Config/TemplateDefs.ini
@@ -16,6 +16,9 @@ LocalizedDescriptions=(Language="en", Text="Blank template with minimal function
 LocalizedDisplayNames=(Language="ko", Text="빈 프로젝트")
 LocalizedDescriptions=(Language="ko", Text="최소한의 기능을 가진 빈 프로젝트 템플릿")
 
+LocalizedDisplayNames=(Language="zh-Hans", Text="空白")
+LocalizedDescriptions=(Language="zh-Hans", Text="具备基本功能的空白模板。")
+
 ; The sort key to use for the template to bypass the default alphabetical sorting.
 SortKey="_1"
 
