Adding Japanese translation #jira ueloc-4712 #fyi michelle.zhao, yali.he

diff --git a/Templates/TP_ME_BlankBP/Config/TemplateDefs.ini b/Templates/TP_ME_BlankBP/Config/TemplateDefs.ini
--- a/Templates/TP_ME_BlankBP/Config/TemplateDefs.ini
+++ b/Templates/TP_ME_BlankBP/Config/TemplateDefs.ini
@@ -19,6 +19,8 @@ LocalizedDescriptions=(Language="ko", Text="최소한의 기능을 가진 빈 
 LocalizedDisplayNames=(Language="zh-Hans", Text="空白")
 LocalizedDescriptions=(Language="zh-Hans", Text="具备基本功能的空白模板。")
 
+LocalizedDescriptions=(Language="ja", Text="最小限の機能が備わった Blank テンプレートです。")
+
 ; The sort key to use for the template to bypass the default alphabetical sorting.
 SortKey="_1"
 
