Removing r.Tonemapper.Sharpen=2 from M&E templates because it causes visual artifacts, the engine default should be used instead

diff --git a/Templates/TP_ME_BlankBP/Config/DefaultEngine.ini b/Templates/TP_ME_BlankBP/Config/DefaultEngine.ini
--- a/Templates/TP_ME_BlankBP/Config/DefaultEngine.ini
+++ b/Templates/TP_ME_BlankBP/Config/DefaultEngine.ini
@@ -8,7 +8,6 @@ AppliedDefaultGraphicsPerformance=Maximum
 
 [/Script/Engine.RendererSettings]
 ; exposure control and tone mapping settings
-r.Tonemapper.Sharpen=2
 r.DefaultFeature.AutoExposure.ExtendDefaultLuminanceRange=True
 r.DefaultFeature.Bloom=False
 r.DefaultFeature.AutoExposure=False
diff --git a/Templates/TP_ME_VProdBP/Config/DefaultEngine.ini b/Templates/TP_ME_VProdBP/Config/DefaultEngine.ini
--- a/Templates/TP_ME_VProdBP/Config/DefaultEngine.ini
+++ b/Templates/TP_ME_VProdBP/Config/DefaultEngine.ini
@@ -7,7 +7,6 @@ AppliedDefaultGraphicsPerformance=Maximum
 
 [/Script/Engine.RendererSettings]
 ; exposure control and tone mapping settings
-r.Tonemapper.Sharpen=2
 r.DefaultFeature.AutoExposure.ExtendDefaultLuminanceRange=True
 r.DefaultFeature.Bloom=False
 r.DefaultFeature.AutoExposure=False
